InEditorDocumentation Plugin

Add plugin source as experimental to UE 5.7
Enable the plugin in the TP_UEIntro_BP template project.

diff --git a/Engine/Plugins/Experimental/InEditorDocumentation/InEditorDocumentation.uplugin b/Engine/Plugins/Experimental/InEditorDocumentation/InEditorDocumentation.uplugin
new file mode 100644
--- /dev/null
+++ b/Engine/Plugins/Experimental/InEditorDocumentation/InEditorDocumentation.uplugin
@@ -0,0 +1,35 @@
+{
+	"FileVersion": 3,
+	"Version": 1,
+	"VersionName": "1.0",
+	"FriendlyName": "In-Editor Documentation",
+	"Description": "Navigate a configured tutorial within the Unreal Editor as you explore a project.",
+	"Category": "Learning",
+	"CreatedBy" : "Epic Games, Inc.",
+	"CreatedByURL" : "https://epicgames.com",
+	"DocsURL": "",
+	"MarketplaceURL": "",
+	"SupportURL": "",
+	"CanContainContent": true,
+	"IsBetaVersion": false,
+	"IsExperimentalVersion": true,
+	"Installed": false,
+	"Modules": [
+		{
+			"Name": "InEditorDocumentation",
+			"Type": "Editor",
+			"LoadingPhase": "Default"
+		},
+		{
+			"Name": "ChangeWebBrowserUserAgent",
+			"Type": "Editor",
+			"LoadingPhase": "PostConfigInit"
+		}
+	],
+	"Plugins": [
+		{
+			"Name": "WebBrowserWidget",
+			"Enabled": true
+		}
+	]
+}
\ No newline at end of file
diff --git a/Engine/Plugins/Experimental/InEditorDocumentation/Resources/page-not-found.html b/Engine/Plugins/Experimental/InEditorDocumentation/Resources/page-not-found.html
new file mode 100644
--- /dev/null
+++ b/Engine/Plugins/Experimental/InEditorDocumentation/Resources/page-not-found.html
@@ -0,0 +1,20 @@
+<!DOCTYPE html>
+<head>
+	<style>
+		.content {
+			text-align: center;
+			vertical-align: middle;
+			padding: 20px;
+			background-color: #242424;
+			font-family: 'Robot', sans-serif;
+			font-size: 10pt;
+			color: #838383;
+		}
+	</style>
+</head>
+<body>
+	<div class="content">
+		<h1 style="font-size: 14pt;">Content not found</h1>
+		<p>The content you are looking for could not be found. Try reloading the page. If the problem persists, check your internet connection. You must be connected to the internet to access this content.</p>
+	</div>
+</body>
\ No newline at end of file
diff --git a/Engine/Plugins/Experimental/InEditorDocumentation/Source/ChangeWebBrowserUserAgent/ChangeWebBrowserUserAgent.Build.cs b/Engine/Plugins/Experimental/InEditorDocumentation/Source/ChangeWebBrowserUserAgent/ChangeWebBrowserUserAgent.Build.cs
new file mode 100644
--- /dev/null
+++ b/Engine/Plugins/Experimental/InEditorDocumentation/Source/ChangeWebBrowserUserAgent/ChangeWebBrowserUserAgent.Build.cs
@@ -0,0 +1,24 @@
+// Copyright Epic Games, Inc. All Rights Reserved.
+
+using UnrealBuildTool;
+
+public class ChangeWebBrowserUserAgent : ModuleRules
+{
+    public ChangeWebBrowserUserAgent(ReadOnlyTargetRules Target) : base(Target)
+    {
+        PCHUsage = PCHUsageMode.UseExplicitOrSharedPCHs;
+
+        PublicDependencyModuleNames.AddRange(new[]
+        {
+            "Core",
+            // add public deps here
+        });
+
+        PrivateDependencyModuleNames.AddRange(new[]
+        {
+            "WebBrowser",
+            "Engine",
+            // add private deps here (InputCore, Projects, etc.)
+        });
+    }
+}
\ No newline at end of file
diff --git a/Engine/Plugins/Experimental/InEditorDocumentation/Source/ChangeWebBrowserUserAgent/ChangeWebBrowserUserAgent.cpp b/Engine/Plugins/Experimental/InEditorDocumentation/Source/ChangeWebBrowserUserAgent/ChangeWebBrowserUserAgent.cpp
new file mode 100644
--- /dev/null
+++ b/Engine/Plugins/Experimental/InEditorDocumentation/Source/ChangeWebBrowserUserAgent/ChangeWebBrowserUserAgent.cpp
@@ -0,0 +1,27 @@
+// Copyright Epic Games, Inc. All Rights Reserved.
+
+#include "ChangeWebBrowserUserAgent.h"
+#include "Modules/ModuleManager.h"
+#include "WebBrowserModule.h"
+
+IMPLEMENT_MODULE(FChangeWebBrowserUserAgent, ChangeWebBrowserUserAgent)
+
+void FChangeWebBrowserUserAgent::StartupModule()
+{
+    // Obtain reference to WebBrowser module
+    IWebBrowserModule& WebBrowserModule = FModuleManager::LoadModuleChecked<IWebBrowserModule>("WebBrowser");
+
+    // Create an empty product version for the User-Agent
+    FWebBrowserInitSettings Init;
+    Init.ProductVersion = TEXT("");
+
+    // Attempt to apply this initialization setting
+    const bool Applied = WebBrowserModule.CustomInitialize(Init);
+
+    // Make sure that the initialization was applied
+    ensureMsgf(Applied, TEXT("Custom WebBrowser singleton initialization not applied"));
+}
+
+void FChangeWebBrowserUserAgent::ShutdownModule()
+{
+}
\ No newline at end of file
diff --git a/Engine/Plugins/Experimental/InEditorDocumentation/Source/ChangeWebBrowserUserAgent/ChangeWebBrowserUserAgent.h b/Engine/Plugins/Experimental/InEditorDocumentation/Source/ChangeWebBrowserUserAgent/ChangeWebBrowserUserAgent.h
new file mode 100644
--- /dev/null
+++ b/Engine/Plugins/Experimental/InEditorDocumentation/Source/ChangeWebBrowserUserAgent/ChangeWebBrowserUserAgent.h
@@ -0,0 +1,18 @@
+// Copyright Epic Games, Inc. All Rights Reserved.
+
+#pragma once
+
+#include "CoreMinimal.h"
+#include "Modules/ModuleManager.h"
+
+/**
+* Module to change the Initialization settings for the Web Browser module.
+* Must load before the WebBrowser module singleton is created.
+* Resets the product version in the browser User-Agent Http request field.
+*/
+class FChangeWebBrowserUserAgent final : public IModuleInterface
+{
+public:
+    virtual void StartupModule() override;
+    virtual void ShutdownModule() override;
+};
\ No newline at end of file
diff --git a/Engine/Plugins/Experimental/InEditorDocumentation/Source/InEditorDocumentation/InEditorDocumentation.Build.cs b/Engine/Plugins/Experimental/InEditorDocumentation/Source/InEditorDocumentation/InEditorDocumentation.Build.cs
new file mode 100644
--- /dev/null
+++ b/Engine/Plugins/Experimental/InEditorDocumentation/Source/InEditorDocumentation/InEditorDocumentation.Build.cs
@@ -0,0 +1,65 @@
+// Copyright Epic Games, Inc. All Rights Reserved.
+
+using UnrealBuildTool;
+
+public class InEditorDocumentation : ModuleRules
+{
+	public InEditorDocumentation(ReadOnlyTargetRules Target) : base(Target)
+	{
+		PCHUsage = ModuleRules.PCHUsageMode.UseExplicitOrSharedPCHs;
+		
+		PublicIncludePaths.AddRange(
+			new string[] {
+				// ... add public include paths required here ...
+			}
+			);
+				
+		
+		PrivateIncludePaths.AddRange(
+			new string[] {
+				// ... add other private include paths required here ...
+			}
+			);
+			
+		
+		PublicDependencyModuleNames.AddRange(
+			new string[]
+			{
+				"Core",
+				// ... add other public dependencies that you statically link with here ...
+			}
+			);
+			
+		
+		PrivateDependencyModuleNames.AddRange(
+			new string[]
+			{
+                "Projects",
+                "InputCore",
+                "EditorFramework",
+                "UnrealEd",
+                "CoreUObject",
+                "Engine",
+                "Slate",
+                "SlateCore",
+                "Blutility",
+                "UMG",
+                "WebBrowser",
+				"ToolMenus",
+				"HTTP",
+				"Json",
+				"JsonUtilities",
+				"LevelEditor",
+				// ... add private dependencies that you statically link with here ...	
+			}
+			);
+		
+		
+		DynamicallyLoadedModuleNames.AddRange(
+			new string[]
+			{
+				// ... add any modules that your module loads dynamically here ...
+			}
+			);
+	}
+}
diff --git a/Engine/Plugins/Experimental/InEditorDocumentation/Source/InEditorDocumentation/Private/CommandsStyle.cpp b/Engine/Plugins/Experimental/InEditorDocumentation/Source/InEditorDocumentation/Private/CommandsStyle.cpp
new file mode 100644
--- /dev/null
+++ b/Engine/Plugins/Experimental/InEditorDocumentation/Source/InEditorDocumentation/Private/CommandsStyle.cpp
@@ -0,0 +1,65 @@
+// Copyright Epic Games, Inc. All Rights Reserved.
+
+
+#include "CommandsStyle.h"
+
+#include "Framework/Application/SlateApplication.h"
+#include "Interfaces/IPluginManager.h"
+#include "Slate/SlateGameResources.h"
+#include "Styling/SlateStyleRegistry.h"
+#include "Styling/SlateStyleMacros.h"
+
+#define RootToContentDir Style->RootToContentDir
+
+TSharedPtr<FSlateStyleSet> FCommandsStyle::StyleInstance = nullptr;
+
+void FCommandsStyle::Initialize()
+{
+	if (!StyleInstance.IsValid())
+	{
+		StyleInstance = Create();
+		FSlateStyleRegistry::RegisterSlateStyle(*StyleInstance);
+	}
+}
+
+void FCommandsStyle::Shutdown()
+{
+	FSlateStyleRegistry::UnRegisterSlateStyle(*StyleInstance);
+	ensure(StyleInstance.IsUnique());
+	StyleInstance.Reset();
+}
+
+void FCommandsStyle::ReloadTextures()
+{
+	if (FSlateApplication::IsInitialized())
+	{
+		FSlateApplication::Get().GetRenderer()->ReloadTextureResources();
+	}
+}
+
+const ISlateStyle& FCommandsStyle::Get()
+{
+	// TODO: insert return statement here
+	return *StyleInstance;
+}
+
+FName FCommandsStyle::GetStyleSetName()
+{
+	static FName StyleSetName(TEXT("CommandsStyle"));
+	return StyleSetName;
+}
+
+const FVector2D Icon20x20(20.0f, 20.0f);
+const FVector2D Icon64x64(64.0f, 64.0f);
+const FVector2D SearchIcon40x40(40.0f, 40.0f);
+
+TSharedRef<class FSlateStyleSet> FCommandsStyle::Create()
+{
+	TSharedRef<FSlateStyleSet> Style = MakeShareable(new FSlateStyleSet("CommandsStyle"));
+	Style->SetContentRoot(IPluginManager::Get().FindPlugin("InEditorDocumentation")->GetBaseDir() / TEXT("Resources"));
+
+	Style->Set("InEditorDocumentation.OpenTutorial", new IMAGE_BRUSH(TEXT("EditorTutorial_64x"), Icon64x64));
+	Style->Set("InEditorDocumentation.OpenSearch", new IMAGE_BRUSH(TEXT("Search_40x"), SearchIcon40x40));
+
+	return Style;
+}
diff --git a/Engine/Plugins/Experimental/InEditorDocumentation/Source/InEditorDocumentation/Private/DocumentationCommands.cpp b/Engine/Plugins/Experimental/InEditorDocumentation/Source/InEditorDocumentation/Private/DocumentationCommands.cpp
new file mode 100644
--- /dev/null
+++ b/Engine/Plugins/Experimental/InEditorDocumentation/Source/InEditorDocumentation/Private/DocumentationCommands.cpp
@@ -0,0 +1,17 @@
+// Copyright Epic Games, Inc. All Rights Reserved.
+
+
+#include "DocumentationCommands.h"
+
+#define LOCTEXT_NAMESPACE "FInEditorDocumentationModule"
+
+void FDocumentationCommands::RegisterCommands()
+{
+	// Command to bring up a documentation window
+	UI_COMMAND(OpenTutorial, "Open Tutorial", "Open Tutorial", EUserInterfaceActionType::Button, FInputChord());
+
+	// Command to search Edc for information
+	UI_COMMAND(OpenSearch, "Search Epic Developer Community", "Search the Epic Developer Community", EUserInterfaceActionType::Button, FInputChord());
+}
+
+#undef LOCTEXT_NAMESPACE
\ No newline at end of file
diff --git a/Engine/Plugins/Experimental/InEditorDocumentation/Source/InEditorDocumentation/Private/InEditorDocumentation.cpp b/Engine/Plugins/Experimental/InEditorDocumentation/Source/InEditorDocumentation/Private/InEditorDocumentation.cpp
new file mode 100644
--- /dev/null
+++ b/Engine/Plugins/Experimental/InEditorDocumentation/Source/InEditorDocumentation/Private/InEditorDocumentation.cpp
@@ -0,0 +1,557 @@
+// Copyright Epic Games, Inc. All Rights Reserved.
+
+#include "InEditorDocumentation.h"
+
+#include "CoreMinimal.h"
+#include "Components/SlateWrapperTypes.h"
+#include "DocumentationCommands.h"
+#include "Framework/Docking/TabManager.h"
+#include "Framework/Docking/LayoutExtender.h"
+#include "Framework/Notifications/NotificationManager.h"
+#include "HttpModule.h"
+#include "Http.h"
+#include "ISettingsModule.h"
+#include "ISettingsSection.h"
+#include "InEditorDocumentationSettings.h"
+#include "Interfaces/IHttpResponse.h"
+#include "Kismet/KismetSystemLibrary.h"
+#include "LevelEditor.h"
+#include "Modules/ModuleManager.h"
+#include "Misc/FileHelper.h"
+#include "Styling/AppStyle.h"
+#include "Selection.h"
+#include "SWebBrowser.h"
+#include "ToolMenus.h"
+#include "WebBrowserModule.h"
+#include "Widgets/Notifications/SNotificationList.h"
+#include "Widgets/Docking/SDockTab.h"
+
+DEFINE_LOG_CATEGORY(LogDocumentation);
+
+static const FName TutorialTabName("Tutorial");
+static const FName SearchTabName("EDCSearch");
+
+#define LOCTEXT_NAMESPACE "FInEditorDocumentationModule"
+
+void FInEditorDocumentationModule::StartupModule()
+{
+	// This code will execute after your module is loaded into memory; the exact timing is specified in the .uplugin file per-module
+
+	// Slate Style for the OpenDocumentationWindow button
+	FCommandsStyle::Initialize();
+	FCommandsStyle::ReloadTextures();
+
+	// UI Commands Init
+	FDocumentationCommands::Register();
+
+	// Initialize the commands
+	PluginCommands = MakeShareable(new FUICommandList);
+
+	check(PluginCommands.IsValid());
+
+	// Map the OpenDocumentationWindow UI Command to the OnOpenDocumentationWindowClicked Callback
+	PluginCommands->MapAction(
+		FDocumentationCommands::Get().OpenTutorial,
+		FExecuteAction::CreateRaw(this, &FInEditorDocumentationModule::OnToggleTutorialClicked)
+	);
+
+	// Map the SearchEdc UI Command to the OnSearchEdcClicked Callback
+	PluginCommands->MapAction(
+		FDocumentationCommands::Get().OpenSearch,
+		FExecuteAction::CreateRaw(this, &FInEditorDocumentationModule::OnSearchClicked)
+	);
+	
+	// Ensure that necessary modules/plugins are registered
+	ISettingsModule& SettingsModule = FModuleManager::LoadModuleChecked<ISettingsModule>(TEXT("Settings"));
+	IWebBrowserModule& WebBrowserModule = FModuleManager::LoadModuleChecked<IWebBrowserModule>(TEXT("WebBrowser"));
+
+	// Register configuration settings for this plugin
+	ModuleConfigSettings = SettingsModule.RegisterSettings(
+		SettingsContainerName,
+		SettingsCategoryName,
+		SettingsSectionName,
+		DisplayName,
+		Description,
+		GetMutableDefault<UInEditorDocumentationSettings>()
+	);
+
+	// If the configuration settings are modified, handle this
+	if (ModuleConfigSettings.IsValid())
+	{
+		ModuleConfigSettings->OnModified().BindRaw(this, &FInEditorDocumentationModule::HandleSettingsModified);
+	}
+
+	// Check the configuration settings to determine whether (experimental) search is enabled
+	const UInEditorDocumentationSettings* Settings = GetDefault<UInEditorDocumentationSettings>();
+	if (Settings)
+	{
+		bShowSearch = Settings->bEnableEdcSearch;
+	}
+
+	// Register the documentation button in the level editor toolbar
+	UToolMenus::RegisterStartupCallback(FSimpleMulticastDelegate::FDelegate::CreateRaw(this, &FInEditorDocumentationModule::RegisterMenus));
+
+	// Register the edc search web browser
+	FGlobalTabmanager::Get()->RegisterNomadTabSpawner(SearchTabName, FOnSpawnTab::CreateRaw(this, &FInEditorDocumentationModule::OnSpawnSearchTab))
+		.SetDisplayName(LOCTEXT("InEditorSearchTabName", "EDC Search"))
+		.SetMenuType(ETabSpawnerMenuType::Hidden);
+
+	// Extend the level editor layout so that the web browsers appear as docked slate tabs
+	ExtendLevelEditorLayout();
+}
+
+void FInEditorDocumentationModule::ShutdownModule()
+{
+	// This function may be called during shutdown to clean up your module.  For modules that support dynamic reloading,
+	// we call this function before unloading the module.
+
+	// Unregister module configuration settings
+	ISettingsModule& SettingsModule = FModuleManager::GetModuleChecked<ISettingsModule>(TEXT("Settings"));
+	SettingsModule.UnregisterSettings(SettingsContainerName, SettingsCategoryName, SettingsSectionName);
+}
+
+void FInEditorDocumentationModule::OnToggleTutorialClicked()
+{
+	// Manage this tab through the Level Editor Tab Manager, not the global tab manager
+	const FLevelEditorModule& LevelEditorModule = FModuleManager::GetModuleChecked<FLevelEditorModule>("LevelEditor");
+
+	// Access the level editor tab manager
+	const TSharedPtr<FTabManager> LevelEditorTabManager = LevelEditorModule.GetLevelEditorTabManager();
+
+	// Try to find a live tab for documentation/tutorial through the tab manager
+	TSharedPtr<SDockTab> FindTabResult = LevelEditorTabManager->FindExistingLiveTab(TutorialTabName);
+
+	if (FindTabResult.IsValid() && TutorialWebBrowser.IsValid())
+	{
+		// If tab is open, try to close it
+		FindTabResult->RequestCloseTab();
+		// Delete the reference to TutorialWebBrowser
+		TutorialWebBrowser.Reset();
+	}
+	else
+	{
+		// If either is not valid, reopen instead
+		LevelEditorTabManager->TryInvokeTab(TutorialTabName);
+	}
+}
+
+void FInEditorDocumentationModule::OnSearchClicked()
+{
+
+	// Try to find a live tab for documentation/tutorial through the tab manager
+	TSharedPtr<SDockTab> FindTabResult = FGlobalTabmanager::Get()->FindExistingLiveTab(SearchTabName);
+
+	if (FindTabResult.IsValid() && SearchWebBrowser.IsValid())
+	{
+		// If the tab and web browser are valid, try to change the url
+		FString Url = ConstructSearchUrl();
+		SearchWebBrowser->LoadURL(Url);
+	}
+	else
+	{
+		// If either is not valid, reopen instead
+		FGlobalTabmanager::Get()->TryInvokeTab(SearchTabName);
+	}
+}
+
+TObjectPtr<AActor> FInEditorDocumentationModule::GetFirstSelectedActor()
+{
+	TArray<TObjectPtr<AActor>> SelectedActors;
+
+	// Access the Unreal Editor global object
+	if (GEditor)
+	{
+		// Iterate through all selected actors in the viewport
+		for (FSelectionIterator It(GEditor->GetSelectedActorIterator()); It; ++It)
+		{
+			// Cast to Actor and add to array
+			if (TObjectPtr<AActor> AsActor = Cast<AActor>(*It))
+			{
+				SelectedActors.Add(AsActor);
+			}
+		}
+	}
+
+	// If there are any, return the first one
+	if (SelectedActors.Num() > 0)
+	{
+		return SelectedActors[0];
+	}
+
+	return nullptr;
+}
+
+void FInEditorDocumentationModule::RegisterMenus()
+{
+	// Owner will be used for cleanup in call to UToolMenus::UnregisterOwner
+	FToolMenuOwnerScoped OwnerScoped(this);
+	{
+		// Get the UI menu we wish to add a button to
+		UToolMenu* ToolbarMenu = UToolMenus::Get()->ExtendMenu("LevelEditor.LevelEditorToolBar.PlayToolBar");
+		{
+			// Get the section within the toolbar menu
+			FToolMenuSection& Section = ToolbarMenu->FindOrAddSection("PluginTools");
+			{
+				// Add our OpenDocumentationWindow UI Command to this menu
+				FToolMenuEntry& OpenDocEntry = Section.AddEntry(FToolMenuEntry::InitToolBarButton(FDocumentationCommands::Get().OpenTutorial));
+				OpenDocEntry.SetCommandList(PluginCommands);
+
+				// Add our SearchEdc UI Command to this menu
+				const FToolMenuEntry& Entry = FToolMenuEntry::InitMenuEntryWithCommandList(FDocumentationCommands::Get().OpenSearch, PluginCommands, TAttribute<FText>(), TAttribute<FText>(), TAttribute<FSlateIcon>(), NAME_None, TOptional<FName>(FName("SearchEdc")));
+				FToolMenuEntry& SearchEdcEntry = Section.AddEntry(Entry);
+
+				// Alter the visibility based on the search settings
+				SearchEdcEntry.Visibility = TAttribute<EVisibility>::CreateLambda([this]()
+					{
+						return bShowSearch ? EVisibility::Visible : EVisibility::Hidden;
+					});
+
+				SearchEntry = &SearchEdcEntry;
+			}
+		}
+	}
+}
+
+void FInEditorDocumentationModule::ExtendLevelEditorLayout()
+{
+	// Ensure that the Level Editor Module is loaded
+	FLevelEditorModule& LevelEditorModule = FModuleManager::LoadModuleChecked<FLevelEditorModule>(TEXT("LevelEditor"));
+
+	LevelEditorModule.OnRegisterTabs().AddLambda([this](TSharedPtr<FTabManager> TabManager)
+		{
+			// Register the tutorial tab spawner
+			TabManager->RegisterTabSpawner(TutorialTabName, FOnSpawnTab::CreateRaw(this, &FInEditorDocumentationModule::OnSpawnTutorialTab))
+				.SetDisplayName(LOCTEXT("InEditorTutorialTabName", "Tutorial"))
+				.SetMenuType(ETabSpawnerMenuType::Hidden);
+		});
+
+	// Extend the layout so the documentation/tutorial can dock next to the level editor viewport
+	LevelEditorModule.OnRegisterLayoutExtensions().AddLambda([](FLayoutExtender& Extender)
+		{
+			Extender.ExtendLayout(LevelEditorTabIds::LevelEditorViewport, ELayoutExtensionPosition::Above, FTabManager::FTab(TutorialTabName, ETabState::OpenedTab));
+		});
+
+	// Once the level editor is created, try to open the documentation window
+	LevelEditorModule.OnLevelEditorCreated().AddLambda([this](TSharedPtr<ILevelEditor> LevelEditorInstance) 
+		{
+			float ReloadDelay = 3.5f;
+			LevelEditorInstance->GetTabManager()->TryInvokeTab(TutorialTabName);
+			FTSTicker::GetCoreTicker().AddTicker(
+				FTickerDelegate::CreateLambda([this](float)
+					{
+						FString TutorialUrl = this->GetTutorialUrl();
+						UE_LOG(LogDocumentation, Display, TEXT("Force tutorial url reload after delay: %s"), *TutorialUrl);
+						this->TutorialWebBrowser->LoadURL(TutorialUrl);
+						return false;
+					}),
+				ReloadDelay
+			);
+		});
+}
+
+FString FInEditorDocumentationModule::GetTutorialUrl()
+{
+	// Default url to start
+	FString Url = TEXT("https://dev.epicgames.com/documentation/unreal-engine/");
+
+	// Access the configuration settings for this plugin
+	TObjectPtr<UInEditorDocumentationSettings> AsModuleSettings;
+	TWeakObjectPtr<UObject> ModuleSettings = ModuleConfigSettings->GetSettingsObject();
+
+	// Default to the url set in configuration
+	if (ModuleSettings.IsValid())
+	{
+		if (AsModuleSettings = Cast<UInEditorDocumentationSettings>(ModuleSettings))
+		{
+			Url = AsModuleSettings->TutorialUrl;
+		}
+	}
+
+	return Url;
+}
+
+void FInEditorDocumentationModule::HandleLoadErrorContent()
+{
+	UE_LOG(LogDocumentation, Display, TEXT("Page load error"))
+	FString NotFoundContentPath = FPaths::ProjectDir() / TEXT("Plugins/InEditorDocumentation/Resources/page-not-found.html");
+	FString Content;
+	if (FFileHelper::LoadFileToString(Content, *NotFoundContentPath))
+	{
+		TutorialWebBrowser->StopLoad();
+		TutorialWebBrowser->LoadString(Content, TEXT("404"));
+	}
+}
+
+
+FString FInEditorDocumentationModule::ConstructSearchUrl()
+{
+	// Default url to start
+	FString Url = TEXT("https://dev.epicgames.com/community/search?types=document&application=unreal_engine&page=1&sort_by=relevancy&published_at_range=any_time");
+
+	// Access the configuration settings for this plugin
+	TObjectPtr<UInEditorDocumentationSettings> AsModuleSettings;
+	TWeakObjectPtr<UObject> ModuleSettings = ModuleConfigSettings->GetSettingsObject();
+
+	// Obtain any selected actors from the viewport
+	TObjectPtr<AActor> SelectedActor = GetFirstSelectedActor();
+
+	// If any actors are selected in the viewport, try to find preconfigured documentation or search the edc
+	if (SelectedActor)
+	{
+		// Get the class name of the selected actor from viewport
+		TObjectPtr<UClass> FirstActorClass = SelectedActor->GetClass();
+		FString ActorClassName = FirstActorClass->GetName().ToLower();
+
+		bool bConfigUrlFound = false;
+		
+		if (ModuleSettings.IsValid())
+		{
+			// Search for the name of the actor's class in preconfigured pages
+			if (AsModuleSettings = Cast<UInEditorDocumentationSettings>(ModuleSettings))
+			{
+				if (AsModuleSettings->DocumentationPages.Contains(ActorClassName))
+				{
+					Url = AsModuleSettings->DocumentationPages[ActorClassName];
+					bConfigUrlFound = true;
+					UE_LOG(LogDocumentation, Display, TEXT("Found class [%s] with url [%s]"), *ActorClassName, *Url)
+				}
+			}
+		}
+
+		// If the actor was not found in configuration, search the edc
+		if (!bConfigUrlFound)
+		{
+			// Construct the search url
+			FString BaseUrl = TEXT("https://dev.epicgames.com/community/search");
+			FString Options = TEXT("&types=document&application=unreal_engine&page=1&sort_by=relevancy&published_at_range=any_time");
+
+			FStringFormatOrderedArguments UrlArguments;
+			UrlArguments.Add(*BaseUrl);
+			UrlArguments.Add(*ActorClassName);
+			UrlArguments.Add(*Options);
+
+			Url = FString::Format(TEXT("{0}?query={1}{2}"), UrlArguments);
+			UE_LOG(LogDocumentation, Display, TEXT("Searching edc for actor [%s] with url [%s]"), *ActorClassName, *Url);
+		}
+	}
+
+	return Url;
+}
+
+TSharedRef<class SDockTab> FInEditorDocumentationModule::OnSpawnTutorialTab(const FSpawnTabArgs& SpawnTabArgs)
+{
+	// Obtain the tutorial url
+	FString Url = GetTutorialUrl();
+
+	// Create web browser and apply relevant settings
+	SAssignNew(TutorialWebBrowser, SWebBrowser)
+		.InitialURL(Url)
+		.ShowControls(true)
+		.ShowAddressBar(false)
+		.ShowErrorMessage(true)
+		.BrowserFrameRate(60)
+		.SupportsTransparency(true)
+		.PopupMenuMethod(EPopupMethod::UseCurrentWindow)
+		.OnBeforeNavigation(SWebBrowser::FOnBeforeBrowse::CreateRaw(this, &FInEditorDocumentationModule::HandleEdcNavigation))
+		.OnLoadError(FSimpleDelegate::CreateRaw(this, &FInEditorDocumentationModule::HandleLoadErrorContent))
+		.OnLoadCompleted(FSimpleDelegate::CreateRaw(this, &FInEditorDocumentationModule::HandleTutorialLoadCompleted));
+	
+
+	return SNew(SDockTab)
+		.TabRole(ETabRole::DocumentTab)
+		.OnTabClosed(SDockTab::FOnTabClosedCallback::CreateRaw(this, &FInEditorDocumentationModule::HandleTutorialTabClosed))
+		[
+			TutorialWebBrowser.ToSharedRef()
+		];
+}
+
+TSharedRef<SDockTab> FInEditorDocumentationModule::OnSpawnSearchTab(const FSpawnTabArgs& SpawnTabArgs)
+{
+	// Obtain the default search url or construct a custom one based on the selected actor in the viewport
+	FString Url = ConstructSearchUrl();
+
+	// Web browser is created and settings applied here
+	SAssignNew(SearchWebBrowser, SWebBrowser)
+		.InitialURL(Url)
+		.ShowControls(true)
+		.ShowAddressBar(false)
+		.PopupMenuMethod(EPopupMethod::UseCurrentWindow)
+		.OnBeforeNavigation(SWebBrowser::FOnBeforeBrowse::CreateRaw(this, &FInEditorDocumentationModule::HandleEdcNavigation))
+		.OnLoadCompleted(FSimpleDelegate::CreateRaw(this, &FInEditorDocumentationModule::HandleSearchLoadCompleted));
+
+	return SNew(SDockTab)
+		.TabRole(ETabRole::NomadTab)
+		[
+			SearchWebBrowser.ToSharedRef()
+		];
+}
+
+void FInEditorDocumentationModule::HandleTutorialTabClosed(TSharedRef<SDockTab> ClosedTab)
+{
+	// If the user has not manually acknowledged and dismissed the tutorial closed tab,
+	// Display a toast so that the user has to acknowledge it
+	if (!bUserDismissedTutorialNotification && !TutorialToast.IsValid())
+	{
+		DisplayTutorialToastNotification();
+	}
+}
+
+bool FInEditorDocumentationModule::HandleEdcNavigation(const FString& NavUrl, const FWebNavigationRequest& Nav)
+{
+	// Only allow links within the EDC in editor
+	if (NavUrl.StartsWith(TEXT("https://dev.epicgames.com/")))
+	{
+		return false;
+	}
+	else
+	{
+		// Open non-EDC links in external web browser
+		UKismetSystemLibrary::LaunchURL(NavUrl);
+
+		return true;
+	}
+}
+
+void FInEditorDocumentationModule::HandleTutorialLoadCompleted()
+{
+	if (this->TutorialWebBrowser.IsValid())
+	{
+		// Disable certain elements for Edc site
+		// This allows for more viewing space within the editor
+		this->TutorialWebBrowser->ExecuteJavascript(TEXT(
+			// Disable site navigation and TOC
+			"document.querySelectorAll('.site-nav').forEach(function(el) { el.style.display = 'none'; });"
+			// Disable header which also contains search bar
+			"document.querySelectorAll('.site-header').forEach(function(el) { el.style.display = 'none'; });"
+			// Disable breadcrumb menu that appears on each page
+			"document.querySelectorAll('.breadcrumb').forEach(function(el) { el.style.display = 'none'; });"
+			// Disable mobile TOC that only appears when window size is narrow
+			"document.querySelectorAll('.is-table-of-contents').forEach(function(el) { el.style.display = 'none'; });"
+		));
+	}
+}
+
+void FInEditorDocumentationModule::HandleSearchLoadCompleted()
+{
+	if (this->SearchWebBrowser.IsValid())
+	{
+		// Disable certain elements for Edc site
+		// This allows for more viewing space within the editor
+		this->SearchWebBrowser->ExecuteJavascript(TEXT(
+			// Disable site navigation and TOC
+			"document.querySelectorAll('.site-nav').forEach(function(el) { el.style.display = 'none'; });"
+			// Disable header which also contains search bar
+			"document.querySelectorAll('.site-header').forEach(function(el) { el.style.display = 'none'; });"
+			// Disable breadcrumb menu that appears on each page
+			"document.querySelectorAll('.breadcrumb').forEach(function(el) { el.style.display = 'none'; });"
+			// Disable search filtering menu
+			"document.querySelectorAll('.col-filters-wrapper').forEach(function(el) { el.style.display = 'none'; });"
+			// Disable mobile TOC that only appears when window size is narrow
+			"document.querySelectorAll('.is-table-of-contents').forEach(function(el) { el.style.display = 'none'; });"
+		));
+	}
+}
+
+void FInEditorDocumentationModule::DisplayTutorialToastNotification()
+{
+	// Construct toast notification
+	FNotificationInfo Notification(FText::FromString(TEXT("Tutorial Closed")));
+	Notification.SubText = FText::FromString(TEXT("To re-open the tutorial, navigate to the Level Editor toolbar and click the Open Tutorial button."));
+	Notification.ExpireDuration = 0.0f;
+	Notification.bFireAndForget = false;
+	Notification.bUseThrobber = true;
+	Notification.Image = FAppStyle::GetBrush("Icons.Warning");
+	Notification.CheckBoxText = FText::FromString(TEXT("Acknowledged"));
+	Notification.CheckBoxStateChanged.BindRaw(this, &FInEditorDocumentationModule::HandleTutorialToastCheckboxStateChanged);
+
+	// Add the toast notification to the manager to display
+	TutorialToast = FSlateNotificationManager::Get().AddNotification(Notification);
+
+	if (TutorialToast.IsValid())
+	{
+		UE_LOG(LogDocumentation, Display, TEXT("Setting toast completion state to pending"));
+		TutorialToast.Pin()->SetCompletionState(SNotificationItem::CS_Pending);
+	}
+	else
+	{
+		UE_LOG(LogDocumentation, Display, TEXT("Toast invalid"));
+	}
+}
+
+void FInEditorDocumentationModule::HandleTutorialToastCheckboxStateChanged(ECheckBoxState NewState)
+{
+	if (NewState == ECheckBoxState::Checked)
+	{
+		UE_LOG(LogDocumentation, Display, TEXT("Checkbox state changed to checked"));
+
+		// If the user has checked the state, record it so the toast does not appear again
+		this->bUserDismissedTutorialNotification = true;
+
+		// If the toast is still valid, set it to complete and remove
+		if (TutorialToast.IsValid())
+		{
+			UE_LOG(LogDocumentation, Display, TEXT("Toast valid, attempting to dismiss"));
+			TutorialToast.Pin()->SetCompletionState(SNotificationItem::CS_Success);
+			TutorialToast.Pin()->ExpireAndFadeout();
+		}
+		else
+		{
+			UE_LOG(LogDocumentation, Display, TEXT("Toast invalid"));
+		}
+	}
+	else
+	{
+		UE_LOG(LogDocumentation, Display, TEXT("Unrecognized checkbox state change"))
+	}
+}
+
+bool FInEditorDocumentationModule::HandleSettingsModified()
+{
+	const UInEditorDocumentationSettings* Settings = GetDefault<UInEditorDocumentationSettings>();
+	if (Settings->bEnableEdcSearch)
+	{
+		UE_LOG(LogDocumentation, Display, TEXT("Edc search enabled"));
+
+		// Make search visible
+		if (SearchEntry)
+		{
+			bShowSearch = true;
+			UToolMenus::Get()->RefreshMenuWidget(FName(TEXT("SearchEdc")));
+		}
+	}
+	else
+	{
+		// Make invisible
+		if (SearchEntry)
+		{
+			UE_LOG(LogDocumentation, Display, TEXT("Found search entry"));
+			bShowSearch = false;
+
+			// Manage this tab through the Level Editor Tab Manager, not the global tab manager
+			const FLevelEditorModule& LevelEditorModule = FModuleManager::GetModuleChecked<FLevelEditorModule>("LevelEditor");
+
+			// Access the level editor tab manager
+			const TSharedPtr<FTabManager> LevelEditorTabManager = LevelEditorModule.GetLevelEditorTabManager();
+
+			// Try to find a live tab for documentation/tutorial through the tab manager
+			TSharedPtr<SDockTab> FindTabResult = LevelEditorTabManager->FindExistingLiveTab(SearchTabName);
+
+			// Try to close the search tab
+			if (FindTabResult.IsValid() && FindTabResult->RequestCloseTab())
+			{
+				UE_LOG(LogDocumentation, Display, TEXT("Closed search tab"))
+			}
+			else
+			{
+				UE_LOG(LogDocumentation, Display, TEXT("Could not close search tab"))
+			}
+
+			UToolMenus::Get()->RefreshMenuWidget(FName(TEXT("SearchEdc")));
+		}
+	}
+	return false;
+}
+
+#undef LOCTEXT_NAMESPACE
+	
+IMPLEMENT_MODULE(FInEditorDocumentationModule, InEditorDocumentation)
\ No newline at end of file
diff --git a/Engine/Plugins/Experimental/InEditorDocumentation/Source/InEditorDocumentation/Private/InEditorDocumentationSettings.cpp b/Engine/Plugins/Experimental/InEditorDocumentation/Source/InEditorDocumentation/Private/InEditorDocumentationSettings.cpp
new file mode 100644
--- /dev/null
+++ b/Engine/Plugins/Experimental/InEditorDocumentation/Source/InEditorDocumentation/Private/InEditorDocumentationSettings.cpp
@@ -0,0 +1,12 @@
+// Copyright Epic Games, Inc. All Rights Reserved.
+
+
+#include "InEditorDocumentationSettings.h"
+
+#if WITH_EDITOR
+void UInEditorDocumentationSettings::PostEditChangeProperty(FPropertyChangedEvent& PropertyChangedEvent)
+{
+	Super::PostEditChangeProperty(PropertyChangedEvent);
+	SaveConfig();
+}
+#endif // WITH_EDITOR
diff --git a/Engine/Plugins/Experimental/InEditorDocumentation/Source/InEditorDocumentation/Public/CommandsStyle.h b/Engine/Plugins/Experimental/InEditorDocumentation/Source/InEditorDocumentation/Public/CommandsStyle.h
new file mode 100644
--- /dev/null
+++ b/Engine/Plugins/Experimental/InEditorDocumentation/Source/InEditorDocumentation/Public/CommandsStyle.h
@@ -0,0 +1,33 @@
+// Copyright Epic Games, Inc. All Rights Reserved.
+
+#pragma once
+
+#include "Styling/SlateStyle.h"
+
+/**
+ * 
+ */
+class INEDITORDOCUMENTATION_API FCommandsStyle
+{
+public:
+
+	static void Initialize();
+
+	static void Shutdown();
+
+	/** reloads textures used by slate renderer */
+	static void ReloadTextures();
+
+	/** @return The Slate style set for the Shooter game */
+	static const ISlateStyle& Get();
+
+	static FName GetStyleSetName();
+
+private:
+
+	static TSharedRef< class FSlateStyleSet > Create();
+
+private:
+
+	static TSharedPtr< class FSlateStyleSet > StyleInstance;
+};
diff --git a/Engine/Plugins/Experimental/InEditorDocumentation/Source/InEditorDocumentation/Public/DocumentationCommands.h b/Engine/Plugins/Experimental/InEditorDocumentation/Source/InEditorDocumentation/Public/DocumentationCommands.h
new file mode 100644
--- /dev/null
+++ b/Engine/Plugins/Experimental/InEditorDocumentation/Source/InEditorDocumentation/Public/DocumentationCommands.h
@@ -0,0 +1,26 @@
+// Copyright Epic Games, Inc. All Rights Reserved.
+
+#pragma once
+
+#include "CommandsStyle.h"
+#include "CoreMinimal.h"
+#include "Framework/Commands/Commands.h"
+
+/**
+ * 
+ */
+class INEDITORDOCUMENTATION_API FDocumentationCommands : public TCommands<FDocumentationCommands>
+{
+public:
+	FDocumentationCommands() : TCommands<FDocumentationCommands>(TEXT("InEditorDocumentation"), NSLOCTEXT("Contexts", "InEditorDocumentation", "InEditorDocumentation Plugin"), NAME_None, FCommandsStyle::GetStyleSetName())
+	{
+	}
+
+	// TCommands<...> implementation
+	virtual void RegisterCommands() override;
+	// TCommands<...> implementation
+
+public:
+	TSharedPtr<FUICommandInfo> OpenTutorial;
+	TSharedPtr<FUICommandInfo> OpenSearch;
+};
\ No newline at end of file
diff --git a/Engine/Plugins/Experimental/InEditorDocumentation/Source/InEditorDocumentation/Public/InEditorDocumentation.h b/Engine/Plugins/Experimental/InEditorDocumentation/Source/InEditorDocumentation/Public/InEditorDocumentation.h
new file mode 100644
--- /dev/null
+++ b/Engine/Plugins/Experimental/InEditorDocumentation/Source/InEditorDocumentation/Public/InEditorDocumentation.h
@@ -0,0 +1,164 @@
+// Copyright Epic Games, Inc. All Rights Reserved.
+
+#pragma once
+
+#include "CoreMinimal.h"
+#include "Modules/ModuleManager.h"
+#include "Styling/SlateTypes.h"
+
+DECLARE_LOG_CATEGORY_EXTERN(LogDocumentation, Display, All);
+
+class AActor;
+class FUICommandList;
+class FJsonObject;
+class ISettingsSection;
+class SDockTab;
+class SNotificationItem;
+class SWebBrowser;
+
+struct FToolMenuEntry;
+struct FWebNavigationRequest;
+
+typedef TSharedPtr<ISettingsSection> ISettingsSectionPtr;
+
+class FInEditorDocumentationModule : public IModuleInterface
+{
+public:
+
+	/** IModuleInterface implementation */
+	virtual void StartupModule() override;
+	virtual void ShutdownModule() override;
+	/** IModuleInterface implementation */
+
+private:
+
+	/* Configuration Section Settings */
+	ISettingsSectionPtr ModuleConfigSettings;
+	const FName SettingsContainerName = "Editor";
+	const FName SettingsCategoryName = "Plugins";
+	const FName SettingsSectionName = "In-Editor Documentation";
+	const FText DisplayName = FText::FromString("In-Editor Documentation");
+	const FText Description = FText::FromString("Settings for documentation in the Unreal Editor");
+
+	// Toast notification
+	TWeakPtr<SNotificationItem> TutorialToast;
+
+	// Has the user manually dismissed the tutorial toast?
+	bool bUserDismissedTutorialNotification = false;
+
+	// Is Edc search enabled?
+	bool bShowSearch = true;
+
+	// Ui commands
+	TSharedPtr<FUICommandList> PluginCommands;
+
+	// Web browser
+	TSharedPtr<SWebBrowser> TutorialWebBrowser;
+
+	// Search web browser
+	TSharedPtr<SWebBrowser> SearchWebBrowser;
+
+	// Edc Search results
+	TSharedPtr<FJsonObject> SearchResults;
+
+	// Search entry in the Level Editor menu
+	FToolMenuEntry* SearchEntry;
+
+	/// <summary>
+	/// Callback - Invoked as a response to UI_COMMAND(OpenDocumentationWindow,...)
+	/// </summary>
+	void OnToggleTutorialClicked();
+
+	/// <summary>
+	/// Callback - Invoked as a response to UI_COMMAND(SearchEdc,...)
+	/// </summary>
+	void OnSearchClicked();
+
+	/// <summary>
+	/// Get an array of all actors currently selected in the Unreal Editor viewport
+	/// </summary>
+	/// <returns></returns>
+	static TObjectPtr<AActor> GetFirstSelectedActor();
+
+	/// <summary>
+	/// Register UI Menus
+	/// </summary>
+	void RegisterMenus();
+
+	/// <summary>
+	/// Extend the level editor layout
+	/// </summary>
+	void ExtendLevelEditorLayout();
+
+	/// <summary>
+	/// Get the url to navigate to
+	/// </summary>
+	/// <returns></returns>
+	FString GetTutorialUrl();
+
+	/// <summary>
+	/// Construct the url for the search edc
+	/// </summary>
+	/// <returns></returns>
+	FString ConstructSearchUrl();
+
+	/// <summary>
+	/// Callback - Invoked as a response to TryInvokeTab() for OpenDocumentationWindow
+	/// </summary>
+	/// <param name="SpawnTabArgs"></param>
+	/// <returns></returns>
+	TSharedRef<SDockTab> OnSpawnTutorialTab(const class FSpawnTabArgs& SpawnTabArgs);
+
+	/// <summary>
+	/// Callback - Invoked as a response to TryInvokeTab() for SearchEdc
+	/// </summary>
+	/// <param name="SpawnTabArgs"></param>
+	/// <returns></returns>
+	TSharedRef<SDockTab> OnSpawnSearchTab(const class FSpawnTabArgs& SpawnTabArgs);
+
+	/// <summary>
+	/// Callback - Invoked as a response to RequestCloseTab() for OpenDocumentationWindow
+	/// </summary>
+	/// <param name="ClosedTab"></param>
+	void HandleTutorialTabClosed(TSharedRef<SDockTab> ClosedTab);
+
+	/// <summary>
+	/// Callback - Invoked as a response to SWebBrowser::OnLoadError
+	/// </summary>
+	void HandleLoadErrorContent();
+
+	/// <summary>
+	/// Callback - Invoked as a response to SWebBrowser::OnBeforeNavigation
+	/// </summary>
+	/// <param name="NavUrl"></param>
+	/// <param name="Nav"></param>
+	/// <returns></returns>
+	bool HandleEdcNavigation(const FString& NavUrl, const struct FWebNavigationRequest& Nav);
+
+	/// <summary>
+	/// Callback - Invoked as a response to SWebBrowser::OnLoadCompleted for the TutorialWebBrowser
+	/// </summary>
+	void HandleTutorialLoadCompleted();
+
+	/// <summary>
+	/// Callback - Invoked as a response to SWebBrowser::OnLoadCompleted for the SearchWebBrowser
+	/// </summary>
+	void HandleSearchLoadCompleted();
+
+	/// <summary>
+	/// Show the toast notification
+	/// </summary>
+	void DisplayTutorialToastNotification();
+
+	/// <summary>
+	/// Callback - Invoked as a response to Taost notification checkbox state changed
+	/// </summary>
+	/// <param name="NewState"></param>
+	void HandleTutorialToastCheckboxStateChanged(ECheckBoxState NewState);
+
+	/// <summary>
+	/// Callback - Invoked as a response to module settings modified
+	/// </summary>
+	/// <returns></returns>
+	bool HandleSettingsModified();
+};
diff --git a/Engine/Plugins/Experimental/InEditorDocumentation/Source/InEditorDocumentation/Public/InEditorDocumentationSettings.h b/Engine/Plugins/Experimental/InEditorDocumentation/Source/InEditorDocumentation/Public/InEditorDocumentationSettings.h
new file mode 100644
--- /dev/null
+++ b/Engine/Plugins/Experimental/InEditorDocumentation/Source/InEditorDocumentation/Public/InEditorDocumentationSettings.h
@@ -0,0 +1,39 @@
+// Copyright Epic Games, Inc. All Rights Reserved.
+
+#pragma once
+
+#include "CoreMinimal.h"
+#include "UObject/NoExportTypes.h"
+#include "InEditorDocumentationSettings.generated.h"
+
+/**
+ * 
+ */
+UCLASS(Config = EditorPerProjectUserSettings, DefaultConfig)
+class INEDITORDOCUMENTATION_API UInEditorDocumentationSettings : public UObject
+{
+	GENERATED_BODY()
+
+public:
+
+	/**
+	UPROPERTY(EditAnywhere, Config, Category = "General", meta=(ToolTip="Default url to open if no search results found"))
+	FString DefaultUrl = TEXT("https://dev.epicgames.com/documentation/unreal-engine");
+	*/
+
+	UPROPERTY(EditAnywhere, Config, Category = "Tutorial", meta=(ToolTip="Tutorial url to open if the tutorial button is clicked"))
+	FString TutorialUrl = TEXT("https://dev.epicgames.com/documentation/unreal-engine/stack-o-bot-sample-game-in-unreal-engine");
+
+	UPROPERTY(EditAnywhere, Config, Category = "Search (Experimental)", meta =(ToolTip="Enable Edc search for actors selected in the viewport"))
+	bool bEnableEdcSearch = false;
+
+	UPROPERTY(VisibleAnywhere, Config, Category = "Search (Experimental)", meta=(ToolTip="Api endpoint to search the edc for information about selected actor"))
+	FString EdcSearchApiEndpoint = TEXT("https://dev.epicgames.com/community/api/search/index.json");
+
+	UPROPERTY(EditAnywhere, Config, Category = "Search (Experimental)", meta=(ToolTip="List of predefined documentation page given an actor"))
+	TMap<FString, FString> DocumentationPages;
+
+#if WITH_EDITOR
+	virtual void PostEditChangeProperty(FPropertyChangedEvent& PropertyChangedEvent) override;
+#endif
+};
\ No newline at end of file
diff --git a/Templates/TP_UEIntro_BP/UEIntro.uproject b/Templates/TP_UEIntro_BP/UEIntro.uproject
--- a/Templates/TP_UEIntro_BP/UEIntro.uproject
+++ b/Templates/TP_UEIntro_BP/UEIntro.uproject
@@ -14,6 +14,16 @@
 		{
 			"Name": "Landmass",
 			"Enabled": true
+		},
+		{
+			"Name": "InEditorDocumentation",
+			"Enabled": true,
+			"TargetAllowList": [
+				"Editor"
+			],
+			"PlatformAllowList": [
+				"Win64"
+			]
 		}
 	]
 }
\ No newline at end of file
