Change some elements of page styling to match the editor a little more closely: font size (also overall rem size), scrollbars, and hiding large link previews.

diff --git a/Engine/Plugins/Experimental/InEditorDocumentation/Source/InEditorDocumentation/Private/InEditorDocumentation.cpp b/Engine/Plugins/Experimental/InEditorDocumentation/Source/InEditorDocumentation/Private/InEditorDocumentation.cpp
--- a/Engine/Plugins/Experimental/InEditorDocumentation/Source/InEditorDocumentation/Private/InEditorDocumentation.cpp
+++ b/Engine/Plugins/Experimental/InEditorDocumentation/Source/InEditorDocumentation/Private/InEditorDocumentation.cpp
@@ -430,6 +430,12 @@ void FInEditorDocumentationModule::HandleTutorialLoadCompleted()
 			"document.querySelectorAll('.breadcrumb').forEach(function(el) { el.style.display = 'none'; });"
 			// Disable mobile TOC that only appears when window size is narrow
 			"document.querySelectorAll('.is-table-of-contents').forEach(function(el) { el.style.display = 'none'; });"
+			// Scale down the overall rem size
+			"document.documentElement.style.fontSize = '13px';"
+			// Change page styles for scrollbars
+			"const styleElement = document.createElement('style');\n\
+			styleElement.textContent = '::-webkit-scrollbar { width:8px; } ::-webkit-scrollbar-thumb { border-radius:8px;background-color:#575757; } ::-webkit-scrollbar-thumb:hover { background-color:#808080; } .glossary-term-popup-detail-wrapper { display: none; }';\n\
+            document.head.appendChild(styleElement);"
 		));
 	}
 }
@@ -451,6 +457,12 @@ void FInEditorDocumentationModule::HandleSearchLoadCompleted()
 			"document.querySelectorAll('.col-filters-wrapper').forEach(function(el) { el.style.display = 'none'; });"
 			// Disable mobile TOC that only appears when window size is narrow
 			"document.querySelectorAll('.is-table-of-contents').forEach(function(el) { el.style.display = 'none'; });"
+			// Scale down the overall rem size
+			"document.documentElement.style.fontSize = '13px';"
+			// Change page styles for scrollbars
+			"const styleElement = document.createElement('style');\n\
+			styleElement.textContent = '::-webkit-scrollbar { width:8px; } ::-webkit-scrollbar-thumb { border-radius:8px;background-color:#575757; } ::-webkit-scrollbar-thumb:hover { background-color:#808080; } .glossary-term-popup-detail-wrapper { display: none; }';\n\
+            document.head.appendChild(styleElement);"
 		));
 	}
 }
