Fix potential issue with top-level const declaration, and update comment.

diff --git a/Engine/Plugins/Experimental/InEditorDocumentation/Source/InEditorDocumentation/Private/InEditorDocumentation.cpp b/Engine/Plugins/Experimental/InEditorDocumentation/Source/InEditorDocumentation/Private/InEditorDocumentation.cpp
--- a/Engine/Plugins/Experimental/InEditorDocumentation/Source/InEditorDocumentation/Private/InEditorDocumentation.cpp
+++ b/Engine/Plugins/Experimental/InEditorDocumentation/Source/InEditorDocumentation/Private/InEditorDocumentation.cpp
@@ -432,10 +432,10 @@ void FInEditorDocumentationModule::HandleTutorialLoadCompleted()
 			"document.querySelectorAll('.is-table-of-contents').forEach(function(el) { el.style.display = 'none'; });"
 			// Scale down the overall rem size
 			"document.documentElement.style.fontSize = '13px';"
-			// Change page styles for scrollbars
-			"const styleElement = document.createElement('style');\n\
+			// Change page styles for scrollbars and hide link popups
+			"(function() { const styleElement = document.createElement('style');\n\
 			styleElement.textContent = '::-webkit-scrollbar { width:8px; } ::-webkit-scrollbar-thumb { border-radius:8px;background-color:#575757; } ::-webkit-scrollbar-thumb:hover { background-color:#808080; } .glossary-term-popup-detail-wrapper { display: none; }';\n\
-            document.head.appendChild(styleElement);"
+            document.head.appendChild(styleElement); })();"
 		));
 	}
 }
@@ -459,10 +459,10 @@ void FInEditorDocumentationModule::HandleSearchLoadCompleted()
 			"document.querySelectorAll('.is-table-of-contents').forEach(function(el) { el.style.display = 'none'; });"
 			// Scale down the overall rem size
 			"document.documentElement.style.fontSize = '13px';"
-			// Change page styles for scrollbars
-			"const styleElement = document.createElement('style');\n\
+			// Change page styles for scrollbars and hide link popups
+			"(function() { const styleElement = document.createElement('style');\n\
 			styleElement.textContent = '::-webkit-scrollbar { width:8px; } ::-webkit-scrollbar-thumb { border-radius:8px;background-color:#575757; } ::-webkit-scrollbar-thumb:hover { background-color:#808080; } .glossary-term-popup-detail-wrapper { display: none; }';\n\
-            document.head.appendChild(styleElement);"
+            document.head.appendChild(styleElement); })();"
 		));
 	}
 }
