Removing ChangeWebBrowserUserAgent to fix crash caused by calling too early.
#rn Due to a known issue, certain project names may cause the tutorial to fail to load.

diff --git a/Engine/Plugins/Experimental/InEditorDocumentation/InEditorDocumentation.uplugin b/Engine/Plugins/Experimental/InEditorDocumentation/InEditorDocumentation.uplugin
--- a/Engine/Plugins/Experimental/InEditorDocumentation/InEditorDocumentation.uplugin
+++ b/Engine/Plugins/Experimental/InEditorDocumentation/InEditorDocumentation.uplugin
@@ -19,11 +19,6 @@
 			"Name": "InEditorDocumentation",
 			"Type": "Editor",
 			"LoadingPhase": "Default"
-		},
-		{
-			"Name": "ChangeWebBrowserUserAgent",
-			"Type": "Editor",
-			"LoadingPhase": "PostConfigInit"
 		}
 	],
 	"Plugins": [
diff --git a/Engine/Plugins/Experimental/InEditorDocumentation/Source/ChangeWebBrowserUserAgent/ChangeWebBrowserUserAgent.Build.cs b/Engine/Plugins/Experimental/InEditorDocumentation/Source/ChangeWebBrowserUserAgent/ChangeWebBrowserUserAgent.Build.cs
deleted file mode 100644
diff --git a/Engine/Plugins/Experimental/InEditorDocumentation/Source/ChangeWebBrowserUserAgent/ChangeWebBrowserUserAgent.cpp b/Engine/Plugins/Experimental/InEditorDocumentation/Source/ChangeWebBrowserUserAgent/ChangeWebBrowserUserAgent.cpp
deleted file mode 100644
diff --git a/Engine/Plugins/Experimental/InEditorDocumentation/Source/ChangeWebBrowserUserAgent/ChangeWebBrowserUserAgent.h b/Engine/Plugins/Experimental/InEditorDocumentation/Source/ChangeWebBrowserUserAgent/ChangeWebBrowserUserAgent.h
deleted file mode 100644
