[InEditorDocumentation] Limit URLs opened to EDC URLs

Restrict URLs opened by the plugin to only URLs in the dev.epicgames.com/ domain.
This is to mitigate a security issue discovered in VLN-8857.

diff --git a/Engine/Plugins/Experimental/InEditorDocumentation/Source/InEditorDocumentation/Private/InEditorDocumentation.cpp b/Engine/Plugins/Experimental/InEditorDocumentation/Source/InEditorDocumentation/Private/InEditorDocumentation.cpp
--- a/Engine/Plugins/Experimental/InEditorDocumentation/Source/InEditorDocumentation/Private/InEditorDocumentation.cpp
+++ b/Engine/Plugins/Experimental/InEditorDocumentation/Source/InEditorDocumentation/Private/InEditorDocumentation.cpp
@@ -253,7 +253,8 @@ void FInEditorDocumentationModule::ExtendLevelEditorLayout()
 FString FInEditorDocumentationModule::GetTutorialUrl()
 {
 	// Default url to start
-	FString Url = TEXT("https://dev.epicgames.com/documentation/unreal-engine/");
+	FString DefaultUrl = TEXT("https://dev.epicgames.com/documentation/unreal-engine/");
+	FString Url = DefaultUrl;
 
 	// Access the configuration settings for this plugin
 	UInEditorDocumentationSettings* AsModuleSettings;
@@ -269,6 +270,11 @@ FString FInEditorDocumentationModule::GetTutorialUrl()
 		}
 	}
 
+	if (!Url.StartsWith(TEXT("https://dev.epicgames.com/")))
+	{
+		UE_LOG(LogDocumentation, Display, TEXT("URL must belong to dev.epicgames.com"));
+		return DefaultUrl;
+	}
 	return Url;
 }
 
@@ -288,7 +294,8 @@ void FInEditorDocumentationModule::HandleLoadErrorContent()
 FString FInEditorDocumentationModule::ConstructSearchUrl()
 {
 	// Default url to start
-	FString Url = TEXT("https://dev.epicgames.com/community/search?types=document&application=unreal_engine&page=1&sort_by=relevancy&published_at_range=any_time");
+	FString DefaultUrl = TEXT("https://dev.epicgames.com/community/search?types=document&application=unreal_engine&page=1&sort_by=relevancy&published_at_range=any_time");
+	FString Url = DefaultUrl;
 
 	// Access the configuration settings for this plugin
 	UInEditorDocumentationSettings* AsModuleSettings;
@@ -338,6 +345,11 @@ FString FInEditorDocumentationModule::ConstructSearchUrl()
 		}
 	}
 
+	if (!Url.StartsWith(TEXT("https://dev.epicgames.com/")))
+	{
+		UE_LOG(LogDocumentation, Display, TEXT("URL must belong to dev.epicgames.com"));
+		return DefaultUrl;
+	}
 	return Url;
 }
 
@@ -408,9 +420,8 @@ bool FInEditorDocumentationModule::HandleEdcNavigation(const FString& NavUrl, co
 	}
 	else
 	{
-		// Open non-EDC links in external web browser
-		UKismetSystemLibrary::LaunchURL(NavUrl);
-
+		// Deny non-EDC links
+		UE_LOG(LogDocumentation, Display, TEXT("URLs external to dev.epicgames.com are disabled"));
 		return true;
 	}
 }
